OidcToken - add provider name to refresh token request

#rnx

diff --git a/Engine/Source/Programs/OidcToken/TokenService.cs b/Engine/Source/Programs/OidcToken/TokenService.cs
--- a/Engine/Source/Programs/OidcToken/TokenService.cs
+++ b/Engine/Source/Programs/OidcToken/TokenService.cs
@@ -16,6 +16,12 @@ using EpicGames.OIDC;
 
 namespace OidcToken
 {
+	public class ZenAddOidcTokenRequest
+	{
+		public string? ProviderName { get; set; }
+		public string? RefreshToken { get; set; }
+	}
+
 	public class TokenService : IHostedService
 	{
 		private readonly ILogger<TokenService> Logger;
@@ -133,7 +139,8 @@ namespace OidcToken
 							client.DefaultRequestHeaders.Accept.Add(new MediaTypeWithQualityHeaderValue("application/json"));
 
 							string zenUrl = $"{Settings.CurrentValue.ZenUrl}/auth/oidc/refreshtoken";
-							HttpContent content = new StringContent(JsonSerializer.Serialize<OidcTokenInfo>(tokenInfo), Encoding.UTF8, "application/json");
+							var request = new ZenAddOidcTokenRequest { ProviderName = providerName, RefreshToken = tokenInfo.RefreshToken };
+							HttpContent content = new StringContent(JsonSerializer.Serialize<ZenAddOidcTokenRequest>(request), Encoding.UTF8, "application/json");
 							HttpResponseMessage result = await client.PostAsync(zenUrl, content);
 								
 							if (result.IsSuccessStatusCode)
