OidcToken - Made the query mode output some more useful logging messages and support running without the outfile. Useful if a user wants to inspect the state of our tokens  without reading json files.

diff --git a/Engine/Source/Programs/OidcToken/TokenService.cs b/Engine/Source/Programs/OidcToken/TokenService.cs
--- a/Engine/Source/Programs/OidcToken/TokenService.cs
+++ b/Engine/Source/Programs/OidcToken/TokenService.cs
@@ -176,11 +176,32 @@ namespace OidcToken
 
 		private async Task OutputStatus(string service, OidcStatus status)
 		{
-			FileInfo fi = new(Settings.CurrentValue.OutFile);
-			Logger.LogInformation("Token status output to \"{OutFile}\"", fi.FullName);
+			if (Settings.CurrentValue.ResultToConsole)
+			{
+				string s = JsonSerializer.Serialize(new TokenStatusFile(service, status));
+				Console.WriteLine(s);
+			}
 
-			await using FileStream fs = fi.Open(FileMode.Create, FileAccess.Write);
-			await JsonSerializer.SerializeAsync<TokenStatusFile>(fs, new TokenStatusFile(service, status));
+			if (status == OidcStatus.NotLoggedIn)
+			{
+				Logger.LogWarning("Token for provider {ProviderName} does not exist or is old.", service);
+			}
+			else
+			{
+				// verify that the refresh token is valid and actually able to generate a access token
+				await TokenManager.GetAccessToken(service);
+				OidcStatus refreshedStatus = TokenManager.GetStatusForProvider(service);
+				Logger.LogInformation("Determined status of provider {ProviderName} was {Status}", service, refreshedStatus);
+			}
+
+			if (!string.IsNullOrEmpty(Settings.CurrentValue.OutFile))
+			{
+				FileInfo fi = new(Settings.CurrentValue.OutFile);
+				Logger.LogInformation("Token output to \"{OutFile}\"", fi.FullName);
+
+				await using FileStream fs = fi.Open(FileMode.Create, FileAccess.Write);
+				await JsonSerializer.SerializeAsync<TokenStatusFile>(fs, new TokenStatusFile(service, status));
+			}
 		}
 
 		public Task StopAsync(CancellationToken cancellationToken)
