Added documentation on how to configure OidcToken

diff --git a/Engine/Source/Programs/OidcToken/readme.md b/Engine/Source/Programs/OidcToken/readme.md
new file mode 100644
--- /dev/null
+++ b/Engine/Source/Programs/OidcToken/readme.md
@@ -0,0 +1,22 @@
+OidcToken is a small tool to help expose the ability to allocate access and refresh tokens from a Oidc compatible Identity Provider.
+
+# Configuration
+Configuration file can be placed under 
+`Engine\Programs\OidcToken\oidc-configuration.json`
+or `<Game>\Programs\OidcToken\oidc-configuration.json`
+
+The configuration file can look like this:
+```
+{
+	"OidcToken": {
+		"Providers": {
+			"MyOwnProvider": {
+				"ServerUri": "https://<url-to-your-provider>",
+				"ClientId": "<unique-id-from-provider-usually-guid>",
+				"DisplayName": "MyOwnProvider",
+				"RedirectUri": "http://localhost:6556/callback", // this needs to match what is configured as the redirect uri for your IdP
+				"Scopes": "openid profile offline_access" // these scopes are the basic ones you will need, some system may require more and they may be named differently
+			}
+		}
+	}
+}
\ No newline at end of file
