Updated version of Microsoft.Extensions.Configuration used by OidcToken to resolve issue after EpicGames.OIDC got upgraded in 30102405

diff --git a/Engine/Source/Programs/OidcToken/OidcToken.csproj b/Engine/Source/Programs/OidcToken/OidcToken.csproj
--- a/Engine/Source/Programs/OidcToken/OidcToken.csproj
+++ b/Engine/Source/Programs/OidcToken/OidcToken.csproj
@@ -14,7 +14,7 @@
   </ItemGroup>
 
   <ItemGroup>
-    <PackageReference Include="Microsoft.Extensions.Configuration" Version="5.0.0" />
+    <PackageReference Include="Microsoft.Extensions.Configuration" Version="6.0.0" />
     <PackageReference Include="Microsoft.Extensions.DependencyInjection" Version="5.0.2" />
     <PackageReference Include="Microsoft.Extensions.Hosting" Version="5.0.0" />
     <PackageReference Include="Microsoft.Extensions.Logging" Version="5.0.0" />
