Horde: Return exit code 11 from OidcToken if the specified Horde connection does not have auth enabled, allowing the editor to continue despite not having a valid OIDC token.

#rnx

diff --git a/Engine/Source/Developer/DesktopPlatform/Private/DesktopPlatformBase.cpp b/Engine/Source/Developer/DesktopPlatform/Private/DesktopPlatformBase.cpp
--- a/Engine/Source/Developer/DesktopPlatform/Private/DesktopPlatformBase.cpp
+++ b/Engine/Source/Developer/DesktopPlatform/Private/DesktopPlatformBase.cpp
@@ -1027,7 +1027,17 @@ bool FDesktopPlatformBase::GetOidcAccessTokenInternal(const FString& RootDir, co
 		UE_LOG(LogDesktopPlatform, Warning, TEXT("Failed to run OidcToken ('%s%s'). ExitCode: %i"), *GetOidcTokenExecutableFilename(RootDir), *Arguments, ExitCode);
 		return false;
 	}
+
+	if (ExitCode == 11)
+	{
+		UE_LOG(LogDesktopPlatform, Display, TEXT("Auth is disabled for connection"));
+
+		OutToken = FString();
+		OutTokenExpiresAt = FDateTime::MaxValue();
 	
+		return true;
+	}
+
 	// Read the file to a string
 	FString TokenText;
 	if(FFileHelper::LoadFileToString(TokenText, *ResultFilePath))
diff --git a/Engine/Source/Programs/OidcToken/Program.cs b/Engine/Source/Programs/OidcToken/Program.cs
--- a/Engine/Source/Programs/OidcToken/Program.cs
+++ b/Engine/Source/Programs/OidcToken/Program.cs
@@ -19,7 +19,7 @@ namespace OidcToken
 {
 	class Program
 	{
-		static async Task Main(string[] args)
+		static async Task<int> Main(string[] args)
 		{
 			if (args.Any(s => s.Equals("--help") || s.Equals("-help")) || args.Length == 0)
 			{
@@ -36,7 +36,7 @@ namespace OidcToken
 				Console.WriteLine(" --Zen [true/false] - If true the resulting refresh token is posted to Zens token endpoints");
 				Console.WriteLine(" --Project <path> - Project can be used to tell oidc token which game its working in to allow us to read game specific settings");
 
-				return;
+				return 0;
 			}
 
 			ConfigurationBuilder configBuilder = new();
@@ -53,6 +53,11 @@ namespace OidcToken
 			if (options.HordeUrl != null)
 			{
 				hordeAuthConfig = ReadHordeConfigurationAsync(options.HordeUrl).Result;
+				if (hordeAuthConfig.IsAnonymous())
+				{
+					// Indicate to the caller that auth is disabled.
+					return 11;
+				}
 			}
 
 			await Host.CreateDefaultBuilder(args)
@@ -117,14 +122,19 @@ namespace OidcToken
 					services.AddHostedService<TokenService>();
 				})
 				.RunConsoleAsync();
+
+			return 0;
 		}
 
 		class GetHordeAuthConfigResponse
 		{
+			public string Method { get; set; } = String.Empty;
 			public string ProfileName { get; set; } = null!;
 			public string ServerUrl { get; set; } = null!;
 			public string ClientId { get; set; } = null!;
 			public List<string> LocalRedirectUrls { get; set; } = new List<string>();
+
+			public bool IsAnonymous() => Method.Equals("Anonymous", StringComparison.OrdinalIgnoreCase);
 		}
 
 		static async Task<GetHordeAuthConfigResponse> ReadHordeConfigurationAsync(Uri hordeUrl)
@@ -146,16 +156,20 @@ namespace OidcToken
 				}
 			}
 
-			string? localRedirectUrl = authConfig.LocalRedirectUrls.FirstOrDefault();
-			if (String.IsNullOrEmpty(authConfig.ServerUrl) || String.IsNullOrEmpty(authConfig.ClientId) || String.IsNullOrEmpty(localRedirectUrl))
+			if (!authConfig.IsAnonymous())
 			{
-				throw new Exception("No auth server configuration found");
-			}
+				string? localRedirectUrl = authConfig.LocalRedirectUrls.FirstOrDefault();
+				if (String.IsNullOrEmpty(authConfig.ServerUrl) || String.IsNullOrEmpty(authConfig.ClientId) || String.IsNullOrEmpty(localRedirectUrl))
+				{
+					throw new Exception("No auth server configuration found");
+				}
 
-			if (String.IsNullOrEmpty(authConfig.ProfileName))
-			{
-				authConfig.ProfileName = hordeUrl.Host.ToString();
+				if (String.IsNullOrEmpty(authConfig.ProfileName))
+				{
+					authConfig.ProfileName = hordeUrl.Host.ToString();
+				}
 			}
+
 			return authConfig;
 		}
 	}
