OidcToken - Fixed issue were exit code was not being correctly set anymore

diff --git a/Engine/Source/Programs/OidcToken/Program.cs b/Engine/Source/Programs/OidcToken/Program.cs
--- a/Engine/Source/Programs/OidcToken/Program.cs
+++ b/Engine/Source/Programs/OidcToken/Program.cs
@@ -142,7 +142,7 @@ namespace OidcToken
 				})
 				.RunConsoleAsync();
 
-			return 0;
+			return Environment.ExitCode;
 		}
 
 		class GetHordeAuthConfigResponse
