OidcToken - Updated documentation with info regarding lack of support for Google based authentication and some suggested workarounds.

diff --git a/Engine/Source/Programs/OidcToken/readme.md b/Engine/Source/Programs/OidcToken/readme.md
--- a/Engine/Source/Programs/OidcToken/readme.md
+++ b/Engine/Source/Programs/OidcToken/readme.md
@@ -87,3 +87,14 @@ Create a new separate app registration for your cooking apps and add a client se
 When creating the oidc-configuration.json you can find the server uri to use under the `Endpoints` button for your app registration, its usually https://login.microsoftonline.com/<directory-tenant-id>/v2.0
 For Client id use the `client id` of the desktop app you created.
 The scope needs to contain the API scope you created in the backend service, so it usually ends up being: `offline_access profile openid api://<api scope guid>/user.access`
+
+## Google
+
+Unfortunatley at this time Google does not support the features we require to able to use them as a auth source directly. Specifically we rely on JWTs for API access control, each oidc provider tends to call this slightly different things but in practice it means we can rely on the IdP to provide access token for individual resources using a JWT token. Instead Google only provide for a OIDC api that enables us to interact with their users information and federate the login which would force us to simply sync the users into our own user management system and provide the tokens from there. This is currently out of scope from how we use these tokens.
+
+A second problem that prevents us from working around this issue is Googles lack of groups in the id tokens which is tracked here:
+https://issuetracker.google.com/issues/133774835
+
+For workarounds there are identity providers (IdPs) that can be used to federate the Google login and then manage api access tokens as we require. Some licensees have used used `Dex` for this. We have used `Identity Server 4` for similar features in the past (before it was discontinued). There are multiple open source IdPs that could likely serve this purpose.
+ 
+There maybe other solutions out there, we have very limited experience with Googles auth so if someone finds a secure and easy way to interact with the their auth we would be happy to accept changes for this.
\ No newline at end of file
