Horde: Pass through OIDC scopes for API authorization

Allows requesting custom delegated scopes for tools using OidcToken, like UGS interacting with Horde.

PR #12482

diff --git a/Engine/Source/Programs/Horde/HordeServer/Server/ServerController.cs b/Engine/Source/Programs/Horde/HordeServer/Server/ServerController.cs
--- a/Engine/Source/Programs/Horde/HordeServer/Server/ServerController.cs
+++ b/Engine/Source/Programs/Horde/HordeServer/Server/ServerController.cs
@@ -134,6 +134,7 @@ namespace HordeServer.Server
 			{
 				response.ServerUrl = settings.OidcAuthority;
 				response.ClientId = settings.OidcClientId;
+				response.Scopes = settings.OidcApiRequestedScopes;
 			}
 			response.LocalRedirectUrls = settings.OidcLocalRedirectUrls;
 			return response;
diff --git a/Engine/Source/Programs/Horde/HordeServer/ServerSettings.cs b/Engine/Source/Programs/Horde/HordeServer/ServerSettings.cs
--- a/Engine/Source/Programs/Horde/HordeServer/ServerSettings.cs
+++ b/Engine/Source/Programs/Horde/HordeServer/ServerSettings.cs
@@ -214,6 +214,11 @@ namespace HordeServer
 		/// </summary>
 		public string[] OidcClaimHordePerforceUserMapping { get; set; } = [];
 		
+		/// <summary>
+		/// API scopes to request when acquiring OIDC access tokens
+		/// </summary>
+		public string[] OidcApiRequestedScopes { get; set; } = [];
+		
 		/// <summary>
 		/// Add common scopes and mappings to above OIDC config fields
 		/// Provided as a workaround since .NET config will only *merge* array entries when combining multiple config sources.
@@ -236,6 +241,7 @@ namespace HordeServer
 			IReadOnlyList<string> defaultScopes = ["profile", "email", "openid"];
 			IReadOnlyList<string> defaultMappings = ["preferred_username", "email"];
 			IReadOnlyList<string> defaultEmailMapping = ["email"];
+			IReadOnlyList<string> defaultApiScopes = ["offline_access", "openid"];
 			
 			if (OidcAddDefaultScopesAndMappings)
 			{
@@ -244,6 +250,7 @@ namespace HordeServer
 				OidcClaimEmailMapping = OidcClaimEmailMapping.Concat(defaultEmailMapping).ToArray();
 				OidcClaimHordeUserMapping = OidcClaimHordeUserMapping.Concat(defaultMappings).ToArray();
 				OidcClaimHordePerforceUserMapping = OidcClaimHordePerforceUserMapping.Concat(defaultMappings).ToArray();
+				OidcApiRequestedScopes = OidcApiRequestedScopes.Concat(defaultApiScopes).ToArray();
 			}
 		}
 
diff --git a/Engine/Source/Programs/OidcToken/OidcToken.csproj b/Engine/Source/Programs/OidcToken/OidcToken.csproj
--- a/Engine/Source/Programs/OidcToken/OidcToken.csproj
+++ b/Engine/Source/Programs/OidcToken/OidcToken.csproj
@@ -2,7 +2,7 @@
 
   <PropertyGroup>
     <OutputType>Exe</OutputType>
-    <TargetFramework>net6.0</TargetFramework>
+    <TargetFramework>net8.0</TargetFramework>
     <Nullable>enable</Nullable>
     <PublishDir>..\..\..\Binaries\DotNET\OidcToken\</PublishDir>
   </PropertyGroup>
@@ -14,12 +14,16 @@
   </ItemGroup>
 
   <ItemGroup>
-    <PackageReference Include="Microsoft.Extensions.Configuration" Version="6.0.0" />
-    <PackageReference Include="Microsoft.Extensions.DependencyInjection" Version="5.0.2" />
-    <PackageReference Include="Microsoft.Extensions.Hosting" Version="5.0.0" />
-    <PackageReference Include="Microsoft.Extensions.Logging" Version="5.0.0" />
-    <PackageReference Include="Microsoft.Extensions.Logging.Console" Version="5.0.0" />
-    <PackageReference Include="Microsoft.Extensions.Options.DataAnnotations" Version="5.0.0" />
+    <PackageReference Include="Microsoft.Extensions.Configuration" Version="8.0.0" />
+    <PackageReference Include="Microsoft.Extensions.DependencyInjection" Version="8.0.1" />
+    <PackageReference Include="Microsoft.Extensions.Hosting" Version="8.0.1" />
+    <PackageReference Include="Microsoft.Extensions.Logging" Version="8.0.1" />
+    <PackageReference Include="Microsoft.Extensions.Logging.Console" Version="8.0.1" />
+    <PackageReference Include="Microsoft.Extensions.Options.DataAnnotations" Version="8.0.0" />
+    <PackageReference Include="Serilog.Extensions.Hosting" Version="8.0.0" />
+    <PackageReference Include="Serilog.Settings.Configuration" Version="8.0.4" />
+    <PackageReference Include="Serilog.Sinks.Console" Version="6.0.0" />
+    <PackageReference Include="Serilog.Sinks.File" Version="6.0.0" />
   </ItemGroup>
 
   <ItemGroup>
diff --git a/Engine/Source/Programs/OidcToken/Program.cs b/Engine/Source/Programs/OidcToken/Program.cs
--- a/Engine/Source/Programs/OidcToken/Program.cs
+++ b/Engine/Source/Programs/OidcToken/Program.cs
@@ -2,10 +2,8 @@
 
 using EpicGames.OIDC;
 using Microsoft.Extensions.Configuration;
-using Microsoft.Extensions.Configuration.Memory;
 using Microsoft.Extensions.DependencyInjection;
 using Microsoft.Extensions.Hosting;
-using Microsoft.Extensions.Logging;
 using System;
 using System.Collections.Generic;
 using System.IO;
@@ -114,6 +112,10 @@ namespace OidcToken
 						values[$"{nameof(OidcTokenOptions.Providers)}:{hordeAuthConfig.ProfileName}:{nameof(ProviderInfo.ServerUri)}"] = hordeAuthConfig.ServerUrl;
 						values[$"{nameof(OidcTokenOptions.Providers)}:{hordeAuthConfig.ProfileName}:{nameof(ProviderInfo.ClientId)}"] = hordeAuthConfig.ClientId;
 						values[$"{nameof(OidcTokenOptions.Providers)}:{hordeAuthConfig.ProfileName}:{nameof(ProviderInfo.RedirectUri)}"] = hordeAuthConfig.LocalRedirectUrls![0];
+						if (hordeAuthConfig.Scopes is { Length: > 0 })
+						{
+							values[$"{nameof(OidcTokenOptions.Providers)}:{hordeAuthConfig.ProfileName}:{nameof(ProviderInfo.Scopes)}"] = String.Join(" ", hordeAuthConfig.Scopes);
+						}
 						serviceConfig = new ConfigurationBuilder().AddInMemoryCollection(values).Build();
 					}
 					else
@@ -152,6 +154,7 @@ namespace OidcToken
 			public string ServerUrl { get; set; } = null!;
 			public string ClientId { get; set; } = null!;
 			public List<string> LocalRedirectUrls { get; set; } = new List<string>();
+			public string[]? Scopes { get; set; }
 
 			public bool IsAnonymous() => Method.Equals("Anonymous", StringComparison.OrdinalIgnoreCase);
 		}
diff --git a/Engine/Source/Programs/Shared/EpicGames.Horde/HordeHttpAuthHandler.cs b/Engine/Source/Programs/Shared/EpicGames.Horde/HordeHttpAuthHandler.cs
--- a/Engine/Source/Programs/Shared/EpicGames.Horde/HordeHttpAuthHandler.cs
+++ b/Engine/Source/Programs/Shared/EpicGames.Horde/HordeHttpAuthHandler.cs
@@ -449,6 +449,10 @@ namespace EpicGames.Horde
 			values[$"Providers:{oidcProvider}:ServerUri"] = authConfig.ServerUrl;
 			values[$"Providers:{oidcProvider}:ClientId"] = authConfig.ClientId;
 			values[$"Providers:{oidcProvider}:RedirectUri"] = localRedirectUrl;
+			if (authConfig.Scopes is { Length: > 0 })
+			{
+				values[$"Providers:{oidcProvider}:Scopes"] = String.Join(" ", authConfig.Scopes);
+			}
 
 			ConfigurationBuilder builder = new();
 			builder.AddInMemoryCollection(values);
diff --git a/Engine/Source/Programs/Shared/EpicGames.Horde/Server/ServerMessages.cs b/Engine/Source/Programs/Shared/EpicGames.Horde/Server/ServerMessages.cs
--- a/Engine/Source/Programs/Shared/EpicGames.Horde/Server/ServerMessages.cs
+++ b/Engine/Source/Programs/Shared/EpicGames.Horde/Server/ServerMessages.cs
@@ -129,6 +129,11 @@ namespace EpicGames.Horde.Server
 		/// Optional redirect url provided to OIDC login for external tools (typically to a local server)
 		/// </summary>
 		public string[]? LocalRedirectUrls { get; set; }
+
+		/// <summary>
+		/// Optional OIDC scopes for authorizing API requests from external tools
+		/// </summary>
+		public string[]? Scopes { get; set; }
 	}
 
 	/// <summary>
