OidcToken - Fix BuildGraph tag references

#Synced-CL 48106205

diff --git a/Engine/Source/Programs/OidcToken/BuildOidcToken.xml b/Engine/Source/Programs/OidcToken/BuildOidcToken.xml
--- a/Engine/Source/Programs/OidcToken/BuildOidcToken.xml
+++ b/Engine/Source/Programs/OidcToken/BuildOidcToken.xml
@@ -81,8 +81,8 @@
 	</Agent>
 	
 	<Agent Name="Submit OidcToken" Type="Win64">
-		<Node Name="Tag OutputFiles" Requires="#OutputFilesWin64;OutputFilesMac;OutputFilesLinux;OutputFilesPortable" Produces="#OutputFiles">
-			<Tag With="#OutputFiles" Files="#OutputFilesWin64;OutputFilesMac;OutputFilesLinux;OutputFilesPortable" />
+		<Node Name="Tag OutputFiles" Requires="#OutputFilesWin64;#OutputFilesMac;#OutputFilesLinux;#OutputFilesPortable" Produces="#OutputFiles">
+			<Tag With="#OutputFiles" Files="#OutputFilesWin64;#OutputFilesMac;#OutputFilesLinux;#OutputFilesPortable" />
 		</Node>
 		
 		<Node Name="Publish OidcToken" Requires="#OutputFiles">
