OidcToken - Fix changelist description in BuildGraph

#Synced-CL 48106205

diff --git a/Engine/Source/Programs/OidcToken/BuildOidcToken.xml b/Engine/Source/Programs/OidcToken/BuildOidcToken.xml
--- a/Engine/Source/Programs/OidcToken/BuildOidcToken.xml
+++ b/Engine/Source/Programs/OidcToken/BuildOidcToken.xml
@@ -86,7 +86,7 @@
 		</Node>
 		
 		<Node Name="Publish OidcToken" Requires="#OutputFiles">
-			<Submit Description="Updating OidcToken from CL $(Change)" Files="#OutputFiles" RevertUnchanged="true" />
+			<Submit Description="Updating OidcToken from CL $(Change)\n\n#okforgithub public" Files="#OutputFiles" RevertUnchanged="true" />
 		</Node>
 	</Agent>
 </BuildGraph>
