OidcToken - Allow for sharing of log files which fixes issues with logs missing when multiple oidc token invocations happened at the same time

#Synced-CL 50250157

diff --git a/Engine/Source/Programs/OidcToken/Program.cs b/Engine/Source/Programs/OidcToken/Program.cs
--- a/Engine/Source/Programs/OidcToken/Program.cs
+++ b/Engine/Source/Programs/OidcToken/Program.cs
@@ -196,15 +196,15 @@ namespace OidcToken
 					// configure logging output directory match expectation per platform
 					if (RuntimeInformation.IsOSPlatform(OSPlatform.Windows))
 					{
-						configuration.WriteTo.File(Path.Combine(Environment.GetFolderPath(Environment.SpecialFolder.LocalApplicationData), "UnrealEngine\\Common\\OidcToken\\Logs\\oidc-token.log"), rollingInterval:RollingInterval.Day, restrictedToMinimumLevel: LogEventLevel.Debug, retainedFileCountLimit: 7);
+						configuration.WriteTo.File(Path.Combine(Environment.GetFolderPath(Environment.SpecialFolder.LocalApplicationData), "UnrealEngine\\Common\\OidcToken\\Logs\\oidc-token.log"), rollingInterval:RollingInterval.Day, restrictedToMinimumLevel: LogEventLevel.Debug, retainedFileCountLimit: 7, shared: true);
 					}
 					else if (RuntimeInformation.IsOSPlatform(OSPlatform.OSX))
 					{
-						configuration.WriteTo.File(Path.Combine(Environment.GetFolderPath(Environment.SpecialFolder.ApplicationData), "Epic/UnrealEngine/Common/OidcToken/Logs/oidc-token.log"), rollingInterval: RollingInterval.Day, restrictedToMinimumLevel: LogEventLevel.Debug, retainedFileCountLimit: 7);
+						configuration.WriteTo.File(Path.Combine(Environment.GetFolderPath(Environment.SpecialFolder.ApplicationData), "Epic/UnrealEngine/Common/OidcToken/Logs/oidc-token.log"), rollingInterval: RollingInterval.Day, restrictedToMinimumLevel: LogEventLevel.Debug, retainedFileCountLimit: 7, shared: true);
 					}
 					else if (RuntimeInformation.IsOSPlatform(OSPlatform.Linux))
 					{
-						configuration.WriteTo.File(Path.Combine(Environment.GetFolderPath(Environment.SpecialFolder.ApplicationData), "Epic/UnrealEngine/Common/OidcToken/Logs/oidc-token.log"), rollingInterval: RollingInterval.Day, restrictedToMinimumLevel: LogEventLevel.Debug, retainedFileCountLimit: 7);
+						configuration.WriteTo.File(Path.Combine(Environment.GetFolderPath(Environment.SpecialFolder.ApplicationData), "Epic/UnrealEngine/Common/OidcToken/Logs/oidc-token.log"), rollingInterval: RollingInterval.Day, restrictedToMinimumLevel: LogEventLevel.Debug, retainedFileCountLimit: 7, shared: true);
 					}
 				})
 				.ConfigureAppConfiguration(builder =>
@@ -345,7 +345,7 @@ namespace OidcToken
 			/// Something went wrong when trying to allocate a new access token. Most often the refresh token is expired, and we are running with --unattended so we can not trigger a new interactive login.
 			/// </summary>
 			UnableToAllocateToken = 10,
-			/// <summary>
+			/// <summary>	
 			/// Auth is disabled according to the remote configuration
 			/// </summary>
 			AuthIsDisabled= 11,
