OidcToken: Update to TargetFramework net10.0

#rnx
#Synced-CL 51592834

diff --git a/Engine/Source/Programs/OidcToken/OidcToken.csproj b/Engine/Source/Programs/OidcToken/OidcToken.csproj
--- a/Engine/Source/Programs/OidcToken/OidcToken.csproj
+++ b/Engine/Source/Programs/OidcToken/OidcToken.csproj
@@ -3,10 +3,9 @@
 
   <PropertyGroup>
     <OutputType>Exe</OutputType>
-    <TargetFramework>net8.0</TargetFramework>
     <Nullable>enable</Nullable>
     <PublishDir>..\..\..\Binaries\DotNET\OidcToken\</PublishDir>
-    <EnableConfigurationBindingGenerator>true</EnableConfigurationBindingGenerator>
+    <EnableConfigurationBindingGenerator>false</EnableConfigurationBindingGenerator>
   </PropertyGroup>
 
   <ItemGroup>
diff --git a/Engine/Source/Programs/OidcToken/TokenService.cs b/Engine/Source/Programs/OidcToken/TokenService.cs
--- a/Engine/Source/Programs/OidcToken/TokenService.cs
+++ b/Engine/Source/Programs/OidcToken/TokenService.cs
@@ -26,7 +26,7 @@ namespace OidcToken
 	{
 	}
 
-	public class TokenResultFile
+	internal class TokenResultFile
 	{
 		public string Token { get; set; }
 		public DateTimeOffset ExpiresAt { get; set; }
@@ -40,7 +40,7 @@ namespace OidcToken
 		}
 	}
 
-	public class TokenStatusFile
+	internal class TokenStatusFile
 	{
 		public string Service { get; set; }
 
@@ -54,13 +54,13 @@ namespace OidcToken
 	}
 
 
-	public class ZenAddOidcTokenRequest
+	internal class ZenAddOidcTokenRequest
 	{
 		public string? ProviderName { get; set; }
 		public string? RefreshToken { get; set; }
 	}
 
-	public class TokenService : IHostedService
+	internal class TokenService : IHostedService
 	{
 		private readonly ILogger<TokenService> Logger;
 		private readonly IHostApplicationLifetime AppLifetime;
@@ -261,11 +261,11 @@ namespace OidcToken
 		}
 	}
 
-	public class UnableToAllocateTokenException : Exception
+	internal class UnableToAllocateTokenException : Exception
 	{
 	}
 
-	public class TokenServiceOptions: IValidatableObject
+	internal class TokenServiceOptions: IValidatableObject
 	{
 		public enum TokenServiceMode
 		{
